Activating truncation warnings as errors for modules:
CookOnTheFly, CookOnTheFlyNetServer, PakFile, PakFileUtilities

diff --git a/Engine/Source/Developer/CookOnTheFlyNetServer/CookOnTheFlyNetServer.Build.cs b/Engine/Source/Developer/CookOnTheFlyNetServer/CookOnTheFlyNetServer.Build.cs
--- a/Engine/Source/Developer/CookOnTheFlyNetServer/CookOnTheFlyNetServer.Build.cs
+++ b/Engine/Source/Developer/CookOnTheFlyNetServer/CookOnTheFlyNetServer.Build.cs
@@ -6,7 +6,8 @@ public class CookOnTheFlyNetServer : ModuleRules
 {
     public CookOnTheFlyNetServer(ReadOnlyTargetRules Target) : base(Target)
     {
-        PrivateDependencyModuleNames.AddRange(
+		UnsafeTypeCastWarningLevel = WarningLevel.Warning;
+		PrivateDependencyModuleNames.AddRange(
                 new string[]
                 {
                     "TargetPlatform"
diff --git a/Engine/Source/Developer/PakFileUtilities/PakFileUtilities.Build.cs b/Engine/Source/Developer/PakFileUtilities/PakFileUtilities.Build.cs
--- a/Engine/Source/Developer/PakFileUtilities/PakFileUtilities.Build.cs
+++ b/Engine/Source/Developer/PakFileUtilities/PakFileUtilities.Build.cs
@@ -6,6 +6,7 @@ public class PakFileUtilities : ModuleRules
 {
 	public PakFileUtilities(ReadOnlyTargetRules Target) : base(Target)
 	{
+		UnsafeTypeCastWarningLevel = WarningLevel.Warning;
 		PrivateDependencyModuleNames.AddRange(new string[] {
 			"Core",
 			"PakFile",
diff --git a/Engine/Source/Runtime/CookOnTheFly/CookOnTheFly.Build.cs b/Engine/Source/Runtime/CookOnTheFly/CookOnTheFly.Build.cs
--- a/Engine/Source/Runtime/CookOnTheFly/CookOnTheFly.Build.cs
+++ b/Engine/Source/Runtime/CookOnTheFly/CookOnTheFly.Build.cs
@@ -6,7 +6,8 @@ public class CookOnTheFly : ModuleRules
 {
 	public CookOnTheFly(ReadOnlyTargetRules Target) : base(Target)
 	{
-        PublicDependencyModuleNames.AddRange(
+		UnsafeTypeCastWarningLevel = WarningLevel.Warning;
+		PublicDependencyModuleNames.AddRange(
             new string[] {
                 "Core",
                 "Sockets",
diff --git a/Engine/Source/Runtime/PakFile/PakFile.Build.cs b/Engine/Source/Runtime/PakFile/PakFile.Build.cs
--- a/Engine/Source/Runtime/PakFile/PakFile.Build.cs
+++ b/Engine/Source/Runtime/PakFile/PakFile.Build.cs
@@ -6,6 +6,7 @@ public class PakFile : ModuleRules
 {
 	public PakFile(ReadOnlyTargetRules Target) : base(Target)
 	{
+		UnsafeTypeCastWarningLevel = WarningLevel.Warning;
 		PrivateDependencyModuleNames.Add("Core");
 		PrivateDependencyModuleNames.Add("CoreUObject");
 		PublicDependencyModuleNames.Add("RSA");
